Disable analysis on some third party modules

#rnx

diff --git a/Engine/Plugins/Animation/RigLogic/Source/RigLogicLibTest/RigLogicLibTest.Build.cs b/Engine/Plugins/Animation/RigLogic/Source/RigLogicLibTest/RigLogicLibTest.Build.cs
--- a/Engine/Plugins/Animation/RigLogic/Source/RigLogicLibTest/RigLogicLibTest.Build.cs
+++ b/Engine/Plugins/Animation/RigLogic/Source/RigLogicLibTest/RigLogicLibTest.Build.cs
@@ -10,6 +10,7 @@ namespace UnrealBuildTool.Rules
 		public RigLogicLibTest(ReadOnlyTargetRules Target) : base(Target)
 		{
 			bUseUnity = false; // A windows include is preprocessing some method names causing compile failures.
+			bDisableStaticAnalysis = true;
 
 			if (Target.Platform == UnrealTargetPlatform.Win64)
 			{
diff --git a/Engine/Plugins/Experimental/NNERuntimeORTCpu/Source/NNERuntimeORTCpu/NNERuntimeORTCpu.Build.cs b/Engine/Plugins/Experimental/NNERuntimeORTCpu/Source/NNERuntimeORTCpu/NNERuntimeORTCpu.Build.cs
--- a/Engine/Plugins/Experimental/NNERuntimeORTCpu/Source/NNERuntimeORTCpu/NNERuntimeORTCpu.Build.cs
+++ b/Engine/Plugins/Experimental/NNERuntimeORTCpu/Source/NNERuntimeORTCpu/NNERuntimeORTCpu.Build.cs
@@ -6,6 +6,8 @@ public class NNERuntimeORTCpu : ModuleRules
 {
 	public NNERuntimeORTCpu( ReadOnlyTargetRules Target ) : base( Target )
 	{
+		bDisableStaticAnalysis = true;
+
 		bool bIsORTSupported = (Target.Platform == UnrealTargetPlatform.Win64 || Target.Platform == UnrealTargetPlatform.Linux || Target.Platform == UnrealTargetPlatform.Mac);
 		
 		PCHUsage = PCHUsageMode.UseExplicitOrSharedPCHs;
