NNE: Extend CreateONNXModelForOperator to fit the upcoming QA framework

Problem:
New QA framework need a new definition for CreateONNXModelForOperator, however we still want current QA framework to work while it is being developped.

Solution:
Add the new CreateONNXModelForOperator definition on top of the new current, the current one will be deleted latter once we use the new QA framework.

Test:
Win64 editor: NNE Smoke tests
Win64 editor: NNE Style transfer demo.

#rnx

diff --git a/Engine/Plugins/Experimental/NNERuntimeORTCpu/Source/NNEUtils/Internal/NNEUtilsModelBuilder.h b/Engine/Plugins/Experimental/NNERuntimeORTCpu/Source/NNEUtils/Internal/NNEUtilsModelBuilder.h
--- a/Engine/Plugins/Experimental/NNERuntimeORTCpu/Source/NNEUtils/Internal/NNEUtilsModelBuilder.h
+++ b/Engine/Plugins/Experimental/NNERuntimeORTCpu/Source/NNEUtils/Internal/NNEUtilsModelBuilder.h
@@ -94,6 +94,14 @@ NNEUTILS_API bool CreateONNXModelForOperator(bool UseVariadicShapeForModel, cons
 	TConstArrayView<FTensor> InWeightTensors, TConstArrayView<TArray<char>> InWeightTensorsData,
 	const UE::NNE::FAttributeMap& Attributes, FNNEModelRaw& ModelData);
 
+/**
+ * Utility functions to create single layer NN for operator testing with optional attributes
+ */
+NNEUTILS_API bool CreateONNXModelForOperator(bool UseVariadicShapeForModel, const FString& OperatorName, 
+	TConstArrayView<FTensor> InInputTensors, TConstArrayView<FTensor> InOutputTensors,
+	TConstArrayView<FTensor> InWeightTensors, TConstArrayView<TConstArrayView<uint8>> InWeightTensorsData,
+	const UE::NNE::FAttributeMap& Attributes, FNNEModelRaw& ModelData);
+
 /**
  * Create an instance of NNE model builder that creates NNE model/format in memory
  */
diff --git a/Engine/Plugins/Experimental/NNERuntimeORTCpu/Source/NNEUtils/Private/NNEUtilsModelBuilderONNX.cpp b/Engine/Plugins/Experimental/NNERuntimeORTCpu/Source/NNEUtils/Private/NNEUtilsModelBuilderONNX.cpp
--- a/Engine/Plugins/Experimental/NNERuntimeORTCpu/Source/NNEUtils/Private/NNEUtilsModelBuilderONNX.cpp
+++ b/Engine/Plugins/Experimental/NNERuntimeORTCpu/Source/NNEUtils/Private/NNEUtilsModelBuilderONNX.cpp
@@ -457,6 +457,125 @@ NNEUTILS_API bool CreateONNXModelForOperator(bool UseVariadicShapeForModel, cons
 	return true;
 }
 
+//
+//
+//
+NNEUTILS_API bool CreateONNXModelForOperator(bool UseVariadicShapeForModel, const FString& OperatorName,
+	TConstArrayView<FTensor> InInputTensors, TConstArrayView<FTensor> InOutputTensors,
+	TConstArrayView<FTensor> InWeightTensors, TConstArrayView<TConstArrayView<uint8>> InWeightTensorsData,
+	const UE::NNE::FAttributeMap& Attributes, FNNEModelRaw& Model)
+{
+	Model = FNNEModelRaw{};
+	
+	int64 IrVersion = OnnxIrVersion;
+	int64 OpsetVersion = OnnxOpsetVersion;
+
+	if (OperatorName == TEXT("BatchNormalization") ||	// current implementation is opset 9 (next version is 14)
+		OperatorName == TEXT("Clip") ||					// current implementation is opset 6 (next version is 11)
+		OperatorName == TEXT("Pad") ||					// current implementation is opset 2 (next version is 11)
+		OperatorName == TEXT("Split") ||				// current implementation is opset 2 (next version is 11)
+		OperatorName == TEXT("Shape") ||				// current implementation is opset 1 (next version is 13)
+		OperatorName == TEXT("Slice") ||				// current implementation is opset 1 (next version is 10)
+		OperatorName == TEXT("Squeeze") ||				// current implementation is opset 1 (next version is 11)
+		OperatorName == TEXT("Unsqueeze") ||			// current implementation is opset 1 (next version is 11)
+		OperatorName == TEXT("Upsample")				// deprecated starting opset 10
+		)				
+	{
+		OpsetVersion = 9;
+	}
+	else
+	if (OperatorName == TEXT("ReduceL1") ||				// current implementation is opset 11 (next version is 13)
+		OperatorName == TEXT("ReduceL2") ||				// current implementation is opset 11 (next version is 13)
+		OperatorName == TEXT("ReduceLogSum") ||			// current implementation is opset 11 (next version is 13)
+		OperatorName == TEXT("ReduceLogSumExp") ||		// current implementation is opset 11 (next version is 13)
+		OperatorName == TEXT("ReduceLogMin") ||			// current implementation is opset 11 (next version is 13)
+		OperatorName == TEXT("ReduceLogMax") ||			// current implementation is opset 11 (next version is 13)
+		OperatorName == TEXT("ReduceMean") ||			// current implementation is opset 11 (next version is 13)
+		OperatorName == TEXT("ReduceProd") ||			// current implementation is opset 11 (next version is 13)
+		OperatorName == TEXT("ReduceSum") ||			// current implementation is opset 11 (next version is 13)
+		OperatorName == TEXT("ReduceSumSquare") ||		// current implementation is opset 11 (next version is 13)
+		OperatorName == TEXT("Resize") ||				// current implementation is opset 11 (next version is 13)
+		OperatorName == TEXT("Squeeze") ||				// current implementation is opset 11 (next version is 13)
+		OperatorName == TEXT("Unsqueeze")				// current implementation is opset 11 (next version is 13)
+		)
+	{
+		OpsetVersion = 11;
+	}
+
+	TUniquePtr<IModelBuilder> Builder(CreateONNXModelBuilder(IrVersion, OpsetVersion));
+
+	Builder->Begin();
+
+	TArray<int32> ShapeForModel;
+	TArray<IModelBuilder::HTensor> InputTensors;
+	
+	for (int32 Idx = 0; Idx < InInputTensors.Num(); ++Idx)
+	{
+		const FTensor& Desc = InInputTensors[Idx];
+		BuildShapeForModel(UseVariadicShapeForModel, Desc.GetShape(), ShapeForModel);
+		IModelBuilder::HTensor Tensor = Builder->AddTensor(Desc.GetName(), Desc.GetDataType(), ShapeForModel);
+
+		InputTensors.Emplace(Tensor);
+		Builder->AddInput(Tensor);
+	}
+
+	TArray<IModelBuilder::HTensor> OutputTensors;
+
+	for (int32 Idx = 0; Idx < InOutputTensors.Num(); ++Idx)
+	{
+		const FTensor& Desc = InOutputTensors[Idx];
+		BuildShapeForModel(UseVariadicShapeForModel, Desc.GetShape(), ShapeForModel);
+		IModelBuilder::HTensor Tensor = Builder->AddTensor(Desc.GetName(), Desc.GetDataType(), ShapeForModel);
+
+		OutputTensors.Emplace(Tensor);
+		Builder->AddOutput(Tensor);
+	}
+
+	checkf(InWeightTensors.Num() == InWeightTensorsData.Num(), TEXT("Invalid weight tensors data"));
+	TArray<IModelBuilder::HTensor> WeightTensors;
+
+	for (int32 Idx = 0; Idx < InWeightTensors.Num(); ++Idx)
+	{
+		const FTensor& Desc = InWeightTensors[Idx];
+		const TConstArrayView<uint8>& Data = InWeightTensorsData[Idx];
+		check(Data.Num() == Desc.GetDataSize());
+		BuildShapeForModel(false, Desc.GetShape(), ShapeForModel);
+		IModelBuilder::HTensor Tensor = Builder->AddTensor(Desc.GetName(), Desc.GetDataType(), ShapeForModel, Data.GetData(), Data.Num());
+
+		WeightTensors.Emplace(Tensor);
+	}
+
+	auto Op = Builder->AddOperator(OperatorName);
+
+	for (int32 Idx = 0; Idx < InputTensors.Num(); ++Idx)
+	{
+		Builder->AddOperatorInput(Op, InputTensors[Idx]);
+	}
+
+	for (int32 Idx = 0; Idx < WeightTensors.Num(); ++Idx)
+	{
+		// For now weights are added after model inputs in the list of
+		// operator inputs. This should be made more flexible if needed by future tests cases.
+		Builder->AddOperatorInput(Op, WeightTensors[Idx]);
+	}
+
+	for (int32 Idx = 0; Idx < OutputTensors.Num(); ++Idx)
+	{
+		Builder->AddOperatorOutput(Op, OutputTensors[Idx]);
+	}
+
+	for (int32 Idx = 0; Idx < Attributes.Num(); ++Idx)
+	{
+		Builder->AddOperatorAttribute(Op, Attributes.GetName(Idx), Attributes.GetAttributeValue(Idx));
+	}
+
+	
+	Builder->End(Model.Data);
+	Model.Format = ENNEInferenceFormat::ONNX;
+	
+	return true;
+}
+
 /** Return instance of ONNX model builder */
 NNEUTILS_API IModelBuilder* CreateONNXModelBuilder(int64 IrVersion, int64 OpsetVersion)
 {
