NNE Remove unnecessary uses of C and C++ standard libraries in ORT runtimes.

Problem:
According to the coding standard, we must not use std. Every std function probably has an unreal corresponding function.
Some usage references are:
- Engine\Plugins\Experimental\NNERuntimeORTCpu\Source\NNERuntimeORTCpu\Private\NNERuntimeORTCpuUtils.h 22 & 186
- Engine\Plugins\Experimental\NNERuntimeORTCpu\Source\NNERuntimeORTCpu\Private\NNERuntimeORTCpuModel.cpp 75, 146 & 273
and the same in ORTGpu.

Solution:
Use alternatives where possible.

Tests:
- NNEIncubatorEditor, Win64: Style transfer, smoke test green
- NNEIncubator, Win64: Style transfer
- Non-Unity compilation test (p4v)

#rnx

diff --git a/Engine/Plugins/Experimental/NNERuntimeORTCpu/Source/NNERuntimeORTCpu/Private/NNERuntimeORTCpuModel.cpp b/Engine/Plugins/Experimental/NNERuntimeORTCpu/Source/NNERuntimeORTCpu/Private/NNERuntimeORTCpuModel.cpp
--- a/Engine/Plugins/Experimental/NNERuntimeORTCpu/Source/NNERuntimeORTCpu/Private/NNERuntimeORTCpuModel.cpp
+++ b/Engine/Plugins/Experimental/NNERuntimeORTCpu/Source/NNERuntimeORTCpu/Private/NNERuntimeORTCpuModel.cpp
@@ -72,11 +72,16 @@ namespace UE::NNERuntimeORTCpu::Private
 			}
 		}
 #if WITH_EDITOR
-		catch (const std::exception& Exception)
+		catch (const Ort::Exception& Exception)
 		{
 			UE_LOG(LogNNE, Error, TEXT("%s"), UTF8_TO_TCHAR(Exception.what()));
 			return false;
 		}
+		catch (...)
+		{
+			UE_LOG(LogNNE, Error, TEXT("Unknown exception!"));
+			return false;
+		}
 #endif //WITH_EDITOR
 
 		bIsLoaded = true;
@@ -143,7 +148,7 @@ namespace UE::NNERuntimeORTCpu::Private
 
 			TensorsORTType.Emplace(ONNXTensorElementDataTypeEnum);
 
-			std::pair<ENNETensorDataType, uint64> TypeAndSize = TranslateTensorTypeORTToNNE(ONNXTensorElementDataTypeEnum);
+			TypeInfoORT TypeInfo = TranslateTensorTypeORTToNNE(ONNXTensorElementDataTypeEnum);
 
 			TArray<int32> ShapeData;
 			ShapeData.Reserve(CurrentTensorInfo.GetShape().size());
@@ -153,9 +158,9 @@ namespace UE::NNERuntimeORTCpu::Private
 			}
 
 			NNE::FSymbolicTensorShape Shape = NNE::FSymbolicTensorShape::Make(ShapeData);
-			NNE::FTensorDesc SymbolicTensorDesc = NNE::FTensorDesc::Make(FString(CurTensorNamePtr), Shape, TypeAndSize.first);
+			NNE::FTensorDesc SymbolicTensorDesc = NNE::FTensorDesc::Make(FString(CurTensorNamePtr), Shape, TypeInfo.DataType);
 
-			check(SymbolicTensorDesc.GetElementByteSize() == TypeAndSize.second);
+			check(SymbolicTensorDesc.GetElementByteSize() == TypeInfo.ElementSize);
 			SymbolicTensorDescs.Emplace(SymbolicTensorDesc);
 		}
 
@@ -268,9 +273,15 @@ namespace UE::NNERuntimeORTCpu::Private
 			}
 		}
 #if WITH_EDITOR
-		catch (const std::exception& Exception)
+		catch (const Ort::Exception& Exception)
 		{
 			UE_LOG(LogNNE, Error, TEXT("%s"), UTF8_TO_TCHAR(Exception.what()));
+			return false;
+		}
+		catch (...)
+		{
+			UE_LOG(LogNNE, Error, TEXT("Unknown exception!"));
+			return false;
 		}
 #endif //WITH_EDITOR
 
diff --git a/Engine/Plugins/Experimental/NNERuntimeORTCpu/Source/NNERuntimeORTCpu/Private/NNERuntimeORTCpuUtils.h b/Engine/Plugins/Experimental/NNERuntimeORTCpu/Source/NNERuntimeORTCpu/Private/NNERuntimeORTCpuUtils.h
--- a/Engine/Plugins/Experimental/NNERuntimeORTCpu/Source/NNERuntimeORTCpu/Private/NNERuntimeORTCpuUtils.h
+++ b/Engine/Plugins/Experimental/NNERuntimeORTCpu/Source/NNERuntimeORTCpu/Private/NNERuntimeORTCpuUtils.h
@@ -2,12 +2,13 @@
 
 #pragma once
 
+#include "Algo/Transform.h"
 #include "HAL/UnrealMemory.h"
 #include "NNE.h"
 #include "NNERuntimeCPU.h"
-#include "NNETypes.h"
 #include "NNETensor.h"
 #include "NNEThirdPartyWarningDisabler.h"
+#include "NNETypes.h"
 NNE_THIRD_PARTY_INCLUDES_START
 #undef check
 #undef TEXT
@@ -18,8 +19,39 @@ DECLARE_STATS_GROUP(TEXT("NNE"), STATGROUP_NNE, STATCAT_Advanced);
 
 namespace UE::NNERuntimeORTCpu::Private
 {
+	namespace OrtHelper
+	{
+		inline TArray<uint32> GetShape(const Ort::Value &OrtTensor)
+		{
+			OrtTensorTypeAndShapeInfo *TypeAndShapeInfoPtr = nullptr;
+			size_t DimensionsCount = 0;
+
+			Ort::ThrowOnError(Ort::GetApi().GetTensorTypeAndShape(OrtTensor, &TypeAndShapeInfoPtr));
+			Ort::ThrowOnError(Ort::GetApi().GetDimensionsCount(TypeAndShapeInfoPtr, &DimensionsCount));
 
-	using TypeInfoORT = std::pair<ENNETensorDataType, uint64>;
+			TArray<int64_t> OrtShape;
+			OrtShape.SetNumUninitialized(DimensionsCount);
+
+			Ort::ThrowOnError(Ort::GetApi().GetDimensions(TypeAndShapeInfoPtr, OrtShape.GetData(), OrtShape.Num()));
+
+			Ort::GetApi().ReleaseTensorTypeAndShapeInfo(TypeAndShapeInfoPtr);
+
+			TArray<uint32> Result;
+			Algo::Transform(OrtShape, Result, [] (int64_t Value)
+			{
+				check(Value >= 0);
+				return (uint32)Value;
+			});
+
+			return Result;
+		}
+	}
+
+	struct TypeInfoORT
+	{
+		ENNETensorDataType DataType = ENNETensorDataType::None;
+		uint64 ElementSize = 0;
+	};
 
 	inline TypeInfoORT TranslateTensorTypeORTToNNE(unsigned int OrtDataType)
 	{
@@ -183,16 +215,9 @@ namespace UE::NNERuntimeORTCpu::Private
 			const NNE::FTensorBindingCPU& Binding = InBindingTensors[Index];
 			const NNE::FTensorDesc& TensorDesc = InTensorDescs[Index];
 			const Ort::Value& OrtTensor = InOrtTensors[Index];
-			const std::vector<int64_t>& OrtShape = OrtTensor.GetTensorTypeAndShapeInfo().GetShape();
-
-			TArray<uint32> ShapeData;
-			for (int32 DimIndex = 0; DimIndex < OrtShape.size(); ++DimIndex)
-			{
-				check(OrtShape[DimIndex] >= 0);
-				ShapeData.Add(OrtShape[DimIndex]);
-			}
+			TArray<uint32> OrtShape = OrtHelper::GetShape(OrtTensor);
 
-			NNE::FTensorShape Shape = NNE::FTensorShape::Make(ShapeData);
+			NNE::FTensorShape Shape = NNE::FTensorShape::Make(OrtShape);
 			NNE::Internal::FTensor Tensor = NNE::Internal::FTensor::Make(TensorDesc.GetName(), Shape, TensorDesc.GetDataType());
 			OutTensors.Emplace(Tensor);
 
diff --git a/Engine/Plugins/Experimental/NNERuntimeORTGpu/Source/NNERuntimeORT/Private/NNERuntimeORTModel.cpp b/Engine/Plugins/Experimental/NNERuntimeORTGpu/Source/NNERuntimeORT/Private/NNERuntimeORTModel.cpp
--- a/Engine/Plugins/Experimental/NNERuntimeORTGpu/Source/NNERuntimeORT/Private/NNERuntimeORTModel.cpp
+++ b/Engine/Plugins/Experimental/NNERuntimeORTGpu/Source/NNERuntimeORT/Private/NNERuntimeORTModel.cpp
@@ -102,11 +102,16 @@ namespace UE::NNERuntimeORT::Private
 			}
 		}
 #if WITH_EDITOR
-		catch (const std::exception& Exception)
+		catch (const Ort::Exception& Exception)
 		{
 			UE_LOG(LogNNE, Error, TEXT("%s"), UTF8_TO_TCHAR(Exception.what()));
 			return false;
 		}
+		catch (...)
+		{
+			UE_LOG(LogNNE, Error, TEXT("Unknown exception!"));
+			return false;
+		}
 #endif //WITH_EDITOR
 
 		bIsLoaded = true;
@@ -168,7 +173,7 @@ namespace UE::NNERuntimeORT::Private
 
 			TensorsORTType.Emplace(ONNXTensorElementDataTypeEnum);
 
-			std::pair<ENNETensorDataType, uint64> TypeAndSize = TranslateTensorTypeORTToNNE(ONNXTensorElementDataTypeEnum);
+			TypeInfoORT TypeInfo = TranslateTensorTypeORTToNNE(ONNXTensorElementDataTypeEnum);
 
 			TArray<int32> ShapeData;
 			ShapeData.Reserve(CurrentTensorInfo.GetShape().size());
@@ -178,9 +183,9 @@ namespace UE::NNERuntimeORT::Private
 			}
 
 			NNE::FSymbolicTensorShape Shape = NNE::FSymbolicTensorShape::Make(ShapeData);
-			NNE::FTensorDesc SymbolicTensorDesc = NNE::FTensorDesc::Make(FString(TensorNames.Last()), Shape, TypeAndSize.first);
+			NNE::FTensorDesc SymbolicTensorDesc = NNE::FTensorDesc::Make(FString(TensorNames.Last()), Shape, TypeInfo.DataType);
 
-			check(SymbolicTensorDesc.GetElementByteSize() == TypeAndSize.second);
+			check(SymbolicTensorDesc.GetElementByteSize() == TypeInfo.ElementSize);
 			SymbolicTensorDescs.Emplace(SymbolicTensorDesc);
 		}
 
@@ -293,9 +298,15 @@ namespace UE::NNERuntimeORT::Private
 			}
 		}
 #if WITH_EDITOR
-		catch (const std::exception& Exception)
+		catch (const Ort::Exception& Exception)
 		{
 			UE_LOG(LogNNE, Error, TEXT("%s"), UTF8_TO_TCHAR(Exception.what()));
+			return false;
+		}
+		catch (...)
+		{
+			UE_LOG(LogNNE, Error, TEXT("Unknown exception!"));
+			return false;
 		}
 #endif //WITH_EDITOR
 
diff --git a/Engine/Plugins/Experimental/NNERuntimeORTGpu/Source/NNERuntimeORT/Private/NNERuntimeORTUtils.h b/Engine/Plugins/Experimental/NNERuntimeORTGpu/Source/NNERuntimeORT/Private/NNERuntimeORTUtils.h
--- a/Engine/Plugins/Experimental/NNERuntimeORTGpu/Source/NNERuntimeORT/Private/NNERuntimeORTUtils.h
+++ b/Engine/Plugins/Experimental/NNERuntimeORTGpu/Source/NNERuntimeORT/Private/NNERuntimeORTUtils.h
@@ -18,8 +18,40 @@ DECLARE_STATS_GROUP(TEXT("NNE"), STATGROUP_NNE, STATCAT_Advanced);
 
 namespace UE::NNERuntimeORT::Private
 {
+	
+	namespace OrtHelper
+	{
+		inline TArray<uint32> GetShape(const Ort::Value &OrtTensor)
+		{
+			OrtTensorTypeAndShapeInfo *TypeAndShapeInfoPtr = nullptr;
+			size_t DimensionsCount = 0;
+
+			Ort::ThrowOnError(Ort::GetApi().GetTensorTypeAndShape(OrtTensor, &TypeAndShapeInfoPtr));
+			Ort::ThrowOnError(Ort::GetApi().GetDimensionsCount(TypeAndShapeInfoPtr, &DimensionsCount));
+
+			TArray<int64_t> OrtShape;
+			OrtShape.SetNumUninitialized(DimensionsCount);
+
+			Ort::ThrowOnError(Ort::GetApi().GetDimensions(TypeAndShapeInfoPtr, OrtShape.GetData(), OrtShape.Num()));
 
-	using TypeInfoORT = std::pair<ENNETensorDataType, uint64>;
+			Ort::GetApi().ReleaseTensorTypeAndShapeInfo(TypeAndShapeInfoPtr);
+
+			TArray<uint32> Result;
+			Algo::Transform(OrtShape, Result, [] (int64_t Value)
+			{
+				check(Value >= 0);
+				return (uint32)Value;
+			});
+
+			return Result;
+		}
+	}
+
+	struct TypeInfoORT
+	{
+		ENNETensorDataType DataType = ENNETensorDataType::None;
+		uint64 ElementSize = 0;
+	};
 
 	inline TypeInfoORT TranslateTensorTypeORTToNNE(unsigned int OrtDataType)
 	{
@@ -183,16 +215,9 @@ namespace UE::NNERuntimeORT::Private
 			const NNE::FTensorBindingGPU& Binding = InBindingTensors[Index];
 			const NNE::FTensorDesc& TensorDesc = InTensorDescs[Index];
 			const Ort::Value& OrtTensor = InOrtTensors[Index];
-			const std::vector<int64_t>& OrtShape = OrtTensor.GetTensorTypeAndShapeInfo().GetShape();
-
-			TArray<uint32> ShapeData;
-			for (int32 DimIndex = 0; DimIndex < OrtShape.size(); ++DimIndex)
-			{
-				check(OrtShape[DimIndex] >= 0);
-				ShapeData.Add(OrtShape[DimIndex]);
-			}
+			TArray<uint32> OrtShape = OrtHelper::GetShape(OrtTensor);
 
-			NNE::FTensorShape Shape = NNE::FTensorShape::Make(ShapeData);
+			NNE::FTensorShape Shape = NNE::FTensorShape::Make(OrtShape);
 			NNE::Internal::FTensor Tensor = NNE::Internal::FTensor::Make(TensorDesc.GetName(), Shape, TensorDesc.GetDataType());
 			OutTensors.Emplace(Tensor);
 
